Core - Added missing SendSeekMessage to the Seek for the writable FStreamingNetworkFileHandle

diff --git a/Engine/Source/Runtime/StreamingFile/Private/StreamingNetworkPlatformFile.cpp b/Engine/Source/Runtime/StreamingFile/Private/StreamingNetworkPlatformFile.cpp
--- a/Engine/Source/Runtime/StreamingFile/Private/StreamingNetworkPlatformFile.cpp
+++ b/Engine/Source/Runtime/StreamingFile/Private/StreamingNetworkPlatformFile.cpp
@@ -79,37 +79,55 @@ public:
 	}
 	virtual bool		Seek(int64 NewPosition) override
 	{
-		if (NewPosition >= 0 && NewPosition <= FileSize)
+		if( bWritable )
 		{
-			if (NewPosition < CacheStart[0] || NewPosition >= CacheEnd[0] || CacheStart[0] == -1)
+			if( NewPosition == FilePos )
 			{
-				if (NewPosition < CacheStart[1] || NewPosition >= CacheEnd[1] || CacheStart[1] == -1)
+				return true;
+			}
+			else if (NewPosition >= 0 && NewPosition <= FileSize)
+			{
+				if (Network.SendSeekMessage(HandleId, NewPosition))
 				{
-					LazySeek = true;
 					FilePos = NewPosition;
-					if (CacheStart[CurrentCache] != -1)
+					return true;
+				}
+			}
+		}
+		else if( bReadable )
+		{
+			if (NewPosition >= 0 && NewPosition <= FileSize)
+			{
+				if (NewPosition < CacheStart[0] || NewPosition >= CacheEnd[0] || CacheStart[0] == -1)
+				{
+					if (NewPosition < CacheStart[1] || NewPosition >= CacheEnd[1] || CacheStart[1] == -1)
 					{
-						CurrentCache++;
-						CurrentCache %= 2;
-						CacheStart[CurrentCache] = -1; // Invalidate the cache
+						LazySeek = true;
+						FilePos = NewPosition;
+						if (CacheStart[CurrentCache] != -1)
+						{
+							CurrentCache++;
+							CurrentCache %= 2;
+							CacheStart[CurrentCache] = -1; // Invalidate the cache
+						}
+						return true;
+					}
+					else
+					{
+						LazySeek = false;
+						FilePos = NewPosition;
+						CurrentCache = 1;
+						return true;
 					}
-					return true;
 				}
 				else
 				{
 					LazySeek = false;
 					FilePos = NewPosition;
-					CurrentCache = 1;
+					CurrentCache = 0;
 					return true;
 				}
 			}
-			else
-			{
-				LazySeek = false;
-				FilePos = NewPosition;
-				CurrentCache = 0;
-				return true;
-			}
 		}
 		return false;
 	}
