StreamingFile: Code cleanup pass

Includes GitHub PR #550 by moritz-wundke  (https://github.com/EpicGames/UnrealEngine/pull/550)

diff --git a/Engine/Source/Runtime/StreamingFile/Private/StreamingNetworkPlatformFile.cpp b/Engine/Source/Runtime/StreamingFile/Private/StreamingNetworkPlatformFile.cpp
--- a/Engine/Source/Runtime/StreamingFile/Private/StreamingNetworkPlatformFile.cpp
+++ b/Engine/Source/Runtime/StreamingFile/Private/StreamingNetworkPlatformFile.cpp
@@ -1,20 +1,25 @@
 // Copyright 1998-2015 Epic Games, Inc. All Rights Reserved.
 
 #include "StreamingFilePrivatePCH.h"
+#include "ServerTOC.h"
 #include "StreamingNetworkPlatformFile.h"
 #include "Sockets.h"
 #include "FileManagerGeneric.h"
 #include "MultichannelTCP.h"
 #include "IPlatformFileModule.h"
 
+
 DEFINE_LOG_CATEGORY(LogStreamingPlatformFile);
 
+
 /**
  * A helper class for wrapping some of the network file payload specifics
  */
-class FStreamingNetworkFileArchive : public FBufferArchive
+class FStreamingNetworkFileArchive
+	: public FBufferArchive
 {
 public:
+
 	FStreamingNetworkFileArchive(uint32 Command)
 	{
 		// make sure the command is at the start
@@ -30,22 +35,25 @@ public:
 	}
 };
 
+
 static const int32 GBufferCacheSize = 64 * 1024;
 
-class FStreamingNetworkFileHandle : public IFileHandle
-{
-	FStreamingNetworkPlatformFile&		Network;
-	FString					Filename;
-	uint64						HandleId;
-	int64						FilePos;
-	int64						FileSize;
-	bool						bWritable;
-	bool						bReadable;
-	uint8						BufferCache[2][GBufferCacheSize];
-	int64						CacheStart[2];
-	int64						CacheEnd[2];
-	bool						LazySeek;
-	int32						CurrentCache;
+
+class FStreamingNetworkFileHandle
+	: public IFileHandle
+{
+	FStreamingNetworkPlatformFile& Network;
+	FString Filename;
+	uint64 HandleId;
+	int64 FilePos;
+	int64 FileSize;
+	bool bWritable;
+	bool bReadable;
+	uint8 BufferCache[2][GBufferCacheSize];
+	int64 CacheStart[2];
+	int64 CacheEnd[2];
+	bool LazySeek;
+	int32 CurrentCache;
 
 public:
 
@@ -69,15 +77,17 @@ public:
 		Network.SendCloseMessage(HandleId);
 	}
 
-	virtual int64		Size() override
+	virtual int64 Size() override
 	{
 		return FileSize;
 	}
-	virtual int64		Tell() override
+
+	virtual int64 Tell() override
 	{
 		return FilePos;
 	}
-	virtual bool		Seek(int64 NewPosition) override
+
+	virtual bool Seek(int64 NewPosition) override
 	{
 		if( bWritable )
 		{
@@ -85,11 +95,13 @@ public:
 			{
 				return true;
 			}
-			else if (NewPosition >= 0 && NewPosition <= FileSize)
+
+			if (NewPosition >= 0 && NewPosition <= FileSize)
 			{
 				if (Network.SendSeekMessage(HandleId, NewPosition))
 				{
 					FilePos = NewPosition;
+
 					return true;
 				}
 			}
@@ -110,6 +122,7 @@ public:
 							CurrentCache %= 2;
 							CacheStart[CurrentCache] = -1; // Invalidate the cache
 						}
+
 						return true;
 					}
 					else
@@ -117,6 +130,7 @@ public:
 						LazySeek = false;
 						FilePos = NewPosition;
 						CurrentCache = 1;
+
 						return true;
 					}
 				}
@@ -125,17 +139,21 @@ public:
 					LazySeek = false;
 					FilePos = NewPosition;
 					CurrentCache = 0;
+
 					return true;
 				}
 			}
 		}
+
 		return false;
 	}
-	virtual bool		SeekFromEnd(int64 NewPositionRelativeToEnd = 0) override
+
+	virtual bool SeekFromEnd(int64 NewPositionRelativeToEnd = 0) override
 	{
 		return Seek(FileSize + NewPositionRelativeToEnd);
 	}
-	virtual bool		Read(uint8* Destination, int64 BytesToRead) override
+
+	virtual bool Read(uint8* Destination, int64 BytesToRead) override
 	{
 		bool Result = false;
 		if (bReadable && BytesToRead >= 0 && BytesToRead + FilePos <= FileSize)
@@ -204,14 +222,17 @@ public:
 						}
 
 						int64 SizeToRead = GBufferCacheSize;
+
 						if (FilePos + SizeToRead > FileSize)
 						{
 							SizeToRead = FileSize-FilePos;
 						}
+
 						if (Network.SendSeekMessage(HandleId, FilePos))
 						{
 							Result = Network.SendReadMessage(HandleId, BufferCache[CurrentCache], SizeToRead);
 						}
+
 						if (Result)
 						{
 							CacheStart[CurrentCache] = FilePos;
@@ -228,11 +249,14 @@ public:
 				}
 			}
 		}
+
 		return Result;
 	}
-	virtual bool		Write(const uint8* Source, int64 BytesToWrite) override
+
+	virtual bool Write(const uint8* Source, int64 BytesToWrite) override
 	{
 		bool Result = false;
+
 		if (bWritable && BytesToWrite >= 0)
 		{
 			if (BytesToWrite == 0)
@@ -242,6 +266,7 @@ public:
 			else
 			{
 				Result = Network.SendWriteMessage(HandleId, Source, BytesToWrite);
+
 				if (Result)
 				{
 					FilePos += BytesToWrite;
@@ -249,13 +274,16 @@ public:
 				}
 			}
 		}
+
 		return Result;
 	}
 };
 
+
 bool FStreamingNetworkPlatformFile::ShouldBeUsed(IPlatformFile* Inner, const TCHAR* CmdLine) const
 {
 	bool bResult = FNetworkPlatformFile::ShouldBeUsed(Inner, CmdLine);
+
 	if (bResult)
 	{
 		bResult = FParse::Param(CmdLine, TEXT("Streaming")) || !FPlatformMisc::SupportsLocalCaching();
@@ -264,6 +292,7 @@ bool FStreamingNetworkPlatformFile::ShouldBeUsed(IPlatformFile* Inner, const TCH
 		// On desktop platforms, assume 'in place' execution - to prevent deletion of content, force streaming,
 		// unless it's explicitly allowed with -allowcaching (for automation tests with staged builds).
 		const bool bAllowCaching = FParse::Param(CmdLine, TEXT("AllowCaching"));
+
 		if (bResult == false && !bAllowCaching)
 		{
 			UE_LOG(LogStreamingPlatformFile, Warning, TEXT("Cooked desktop platforms do not support non-streaming. Forcing streaming on."));
@@ -271,16 +300,19 @@ bool FStreamingNetworkPlatformFile::ShouldBeUsed(IPlatformFile* Inner, const TCH
 		bResult = bResult || !bAllowCaching;
 #endif
 	}
+
 	return bResult;
 }
 
+
 bool FStreamingNetworkPlatformFile::InitializeInternal(IPlatformFile* Inner, const TCHAR* HostIP)
 {
 	// look for the commandline that will read files from over the network
-	if (HostIP == NULL)
+	if (HostIP == nullptr)
 	{
 		UE_LOG(LogStreamingPlatformFile, Error, TEXT("No Host IP specified in the commandline."));
 		bIsUsable = false;
+
 		return false;
 	}
 
@@ -292,50 +324,52 @@ bool FStreamingNetworkPlatformFile::InitializeInternal(IPlatformFile* Inner, con
 		FileServerPort = OverridePort;
 	}
 
-		// Send the filenames and timestamps to the server.
-		FNetworkFileArchive Payload(NFS_Messages::GetFileList);
-		FillGetFileList(Payload, true);
+	// Send the filenames and timestamps to the server.
+	FNetworkFileArchive Payload(NFS_Messages::GetFileList);
+	FillGetFileList(Payload, true);
 
-		// Send the directories over, and wait for a response.
-		FArrayReader Response;
+	// Send the directories over, and wait for a response.
+	FArrayReader Response;
+
+	if(SendPayloadAndReceiveResponse(Payload,Response))
+	{
+		// Receive the cooked version information.
+		int32 ServerPackageVersion = 0;
+		int32 ServerPackageLicenseeVersion = 0;
+		ProcessServerInitialResponse(Response, ServerPackageVersion, ServerPackageLicenseeVersion);
 
-		if(SendPayloadAndReceiveResponse(Payload,Response))
+		// Make sure we can sync a file.
+		FString TestSyncFile = FPaths::Combine(*(FPaths::EngineDir()), TEXT("Config/BaseEngine.ini"));
+		IFileHandle* TestFileHandle = OpenRead(*TestSyncFile);
+		if (TestFileHandle != nullptr)
 		{
-			// Receive the cooked version information.
-			int32 ServerPackageVersion = 0;
-			int32 ServerPackageLicenseeVersion = 0;
-			ProcessServerInitialResponse(Response, ServerPackageVersion, ServerPackageLicenseeVersion);
-
-			// Make sure we can sync a file.
-			FString TestSyncFile = FPaths::Combine(*(FPaths::EngineDir()), TEXT("Config/BaseEngine.ini"));
-			IFileHandle* TestFileHandle = OpenRead(*TestSyncFile);
-			if (TestFileHandle != NULL)
-			{
-				uint8* FileContents = (uint8*)FMemory::Malloc(TestFileHandle->Size());
-				if (!TestFileHandle->Read(FileContents, TestFileHandle->Size()))
-				{
-					UE_LOG(LogStreamingPlatformFile, Fatal, TEXT("Could not read test file %s."), *TestSyncFile);
-				}
-				FMemory::Free(FileContents);
-				delete TestFileHandle;
-			}
-			else
+			uint8* FileContents = (uint8*)FMemory::Malloc(TestFileHandle->Size());
+			if (!TestFileHandle->Read(FileContents, TestFileHandle->Size()))
 			{
-				UE_LOG(LogStreamingPlatformFile, Fatal, TEXT("Could not open test file %s."), *TestSyncFile);
+				UE_LOG(LogStreamingPlatformFile, Fatal, TEXT("Could not read test file %s."), *TestSyncFile);
 			}
+			FMemory::Free(FileContents);
+			delete TestFileHandle;
+		}
+		else
+		{
+			UE_LOG(LogStreamingPlatformFile, Fatal, TEXT("Could not open test file %s."), *TestSyncFile);
+		}
 
-			FCommandLine::AddToSubprocessCommandline( *FString::Printf( TEXT("-StreamingHostIP=%s"), HostIP ) );
+		FCommandLine::AddToSubprocessCommandline( *FString::Printf( TEXT("-StreamingHostIP=%s"), HostIP ) );
 
-			return true; 
-		}
+		return true; 
+	}
 
-		return false; 
+	return false; 
 }
 
+
 FStreamingNetworkPlatformFile::~FStreamingNetworkPlatformFile()
 {
 }
 
+
 bool FStreamingNetworkPlatformFile::DeleteFile(const TCHAR* Filename)
 {
 	FScopeLock ScopeLock(&SynchronizationObject);
@@ -346,15 +380,28 @@ bool FStreamingNetworkPlatformFile::DeleteFile(const TCHAR* Filename)
 	Payload << RelativeFilename;
 
 	FArrayReader Response;
+
 	if (SendPayloadAndReceiveResponse(Payload, Response) == false)
 	{
 		return false;
 	}
+
 	uint32 Success = 0;
 	Response << Success;
+
 	return !!Success;
 }
 
+
+bool FStreamingNetworkPlatformFile::IsReadOnly(const TCHAR* Filename)
+{
+	FFileInfo Info;
+	GetFileInfo(Filename, Info);
+
+	return Info.ReadOnly;
+}
+
+
 bool FStreamingNetworkPlatformFile::MoveFile(const TCHAR* To, const TCHAR* From)
 {
 	FScopeLock ScopeLock(&SynchronizationObject);
@@ -369,15 +416,19 @@ bool FStreamingNetworkPlatformFile::MoveFile(const TCHAR* To, const TCHAR* From)
 	Payload << RelativeTo;
 
 	FArrayReader Response;
+
 	if (SendPayloadAndReceiveResponse(Payload, Response) == false)
 	{
 		return false;
 	}
+
 	uint32 Success = 0;
 	Response << Success;
+
 	return !!Success;
 }
 
+
 bool FStreamingNetworkPlatformFile::SetReadOnly(const TCHAR* Filename, bool bNewReadOnlyValue)
 {
 	FScopeLock ScopeLock(&SynchronizationObject);
@@ -391,15 +442,28 @@ bool FStreamingNetworkPlatformFile::SetReadOnly(const TCHAR* Filename, bool bNew
 
 	// perform a local operation
 	FArrayReader Response;
+
 	if (SendPayloadAndReceiveResponse(Payload, Response) == false)
 	{
 		return false;
 	}
+
 	bool bSuccess = 0;
 	Response << bSuccess;
+
 	return bSuccess;
 }
 
+
+FDateTime FStreamingNetworkPlatformFile::GetTimeStamp(const TCHAR* Filename)
+{
+	FFileInfo Info;
+	GetFileInfo(Filename, Info);
+
+	return Info.TimeStamp;
+}
+
+
 void FStreamingNetworkPlatformFile::SetTimeStamp(const TCHAR* Filename, FDateTime DateTime)
 {
 	FScopeLock ScopeLock(&SynchronizationObject);
@@ -412,28 +476,43 @@ void FStreamingNetworkPlatformFile::SetTimeStamp(const TCHAR* Filename, FDateTim
 
 	// perform a local operation
 	FArrayReader Response;
+
 	if (SendPayloadAndReceiveResponse(Payload, Response) == false)
 	{
 		return;
 	}
+
 	bool bSuccess = 0;
 	Response << bSuccess;
 }
 
+
+FDateTime FStreamingNetworkPlatformFile::GetAccessTimeStamp(const TCHAR* Filename)
+{
+	FFileInfo Info;
+	GetFileInfo(Filename, Info);
+
+	return Info.AccessTimeStamp;
+}
+
+
 IFileHandle* FStreamingNetworkPlatformFile::OpenRead(const TCHAR* Filename)
 {
 	FString RelativeFilename = Filename;
 	FPaths::MakeStandardFilename(RelativeFilename);
 
 	FStreamingNetworkFileHandle* FileHandle = SendOpenMessage(RelativeFilename, false, false, false);
+
 	return FileHandle;
 }
 
+
 IFileHandle* FStreamingNetworkPlatformFile::OpenWrite(const TCHAR* Filename, bool bAppend, bool bAllowRead)
 {
 	FString RelativeFilename = Filename;
 	FPaths::MakeStandardFilename(RelativeFilename);
 	FStreamingNetworkFileHandle* FileHandle = SendOpenMessage(RelativeFilename, true, bAppend, bAllowRead);
+
 	return FileHandle;
 }
 
@@ -442,6 +521,7 @@ bool FStreamingNetworkPlatformFile::CreateDirectoryTree(const TCHAR* Directory)
 	return IPlatformFile::CreateDirectoryTree( Directory );
 }
 
+
 bool FStreamingNetworkPlatformFile::CreateDirectory(const TCHAR* Directory)
 {
 	FScopeLock ScopeLock(&SynchronizationObject);
@@ -453,15 +533,19 @@ bool FStreamingNetworkPlatformFile::CreateDirectory(const TCHAR* Directory)
 
 	// perform a local operation
 	FArrayReader Response;
+
 	if (SendPayloadAndReceiveResponse(Payload, Response) == false)
 	{
 		return false;
 	}
+
 	bool bSuccess = 0;
 	Response << bSuccess;
+
 	return bSuccess;
 }
 
+
 bool FStreamingNetworkPlatformFile::DeleteDirectory(const TCHAR* Directory)
 {
 	FScopeLock ScopeLock(&SynchronizationObject);
@@ -473,15 +557,19 @@ bool FStreamingNetworkPlatformFile::DeleteDirectory(const TCHAR* Directory)
 
 	// perform a local operation
 	FArrayReader Response;
+
 	if (SendPayloadAndReceiveResponse(Payload, Response) == false)
 	{
 		return false;
 	}
+
 	bool bSuccess = 0;
 	Response << bSuccess;
+
 	return bSuccess;
 }
 
+
 bool FStreamingNetworkPlatformFile::IterateDirectory(const TCHAR* InDirectory, IPlatformFile::FDirectoryVisitor& Visitor)
 {
 	FString RelativeDirectory = InDirectory;
@@ -495,7 +583,7 @@ bool FStreamingNetworkPlatformFile::IterateDirectory(const TCHAR* InDirectory, I
 
 	// Find the directory in TOC
 	FServerTOC::FDirectory* ServerDirectory = ServerFiles.FindDirectory(RelativeDirectory);
-	if (ServerDirectory != NULL)
+	if (ServerDirectory != nullptr)
 	{
 		// loop over the server files and look if they are in this exact directory
 		for (FServerTOC::FDirectory::TIterator It(*ServerDirectory); It && RetVal == true; ++It)
@@ -514,6 +602,7 @@ bool FStreamingNetworkPlatformFile::IterateDirectory(const TCHAR* InDirectory, I
 	return RetVal;
 }
 
+
 bool FStreamingNetworkPlatformFile::IterateDirectoryRecursively(const TCHAR* InDirectory, IPlatformFile::FDirectoryVisitor& Visitor)
 {
 	FString RelativeDirectory = InDirectory;
@@ -539,9 +628,11 @@ bool FStreamingNetworkPlatformFile::IterateDirectoryRecursively(const TCHAR* InD
 			}
 		}
 	}
+
 	return RetVal;
 }
 
+
 bool FStreamingNetworkPlatformFile::DeleteDirectoryRecursively(const TCHAR* Directory)
 {
 	FScopeLock ScopeLock(&SynchronizationObject);
@@ -553,15 +644,19 @@ bool FStreamingNetworkPlatformFile::DeleteDirectoryRecursively(const TCHAR* Dire
 
 	// perform a local operation
 	FArrayReader Response;
+
 	if (SendPayloadAndReceiveResponse(Payload, Response) == false)
 	{
 		return false;
 	}
+
 	bool bSuccess = 0;
 	Response << bSuccess;
+
 	return bSuccess;
 }
 
+
 bool FStreamingNetworkPlatformFile::CopyFile(const TCHAR* To, const TCHAR* From)
 {
 	FScopeLock ScopeLock(&SynchronizationObject);
@@ -574,15 +669,19 @@ bool FStreamingNetworkPlatformFile::CopyFile(const TCHAR* To, const TCHAR* From)
 
 	// perform a local operation
 	FArrayReader Response;
+
 	if (SendPayloadAndReceiveResponse(Payload, Response) == false)
 	{
 		return false;
 	}
+
 	bool bSuccess = 0;
 	Response << bSuccess;
+
 	return bSuccess;
 }
 
+
 FString FStreamingNetworkPlatformFile::ConvertToAbsolutePathForExternalAppForRead( const TCHAR* Filename )
 {
 	FScopeLock ScopeLock(&SynchronizationObject);
@@ -594,13 +693,16 @@ FString FStreamingNetworkPlatformFile::ConvertToAbsolutePathForExternalAppForRea
 
 	// perform a local operation
 	FArrayReader Response;
+
 	if (SendPayloadAndReceiveResponse(Payload, Response) == true)
 	{
 		Response << RelativeFilename;
 	}
+
 	return RelativeFilename;
 }
 
+
 FString FStreamingNetworkPlatformFile::ConvertToAbsolutePathForExternalAppForWrite( const TCHAR* Filename )
 {
 	FScopeLock ScopeLock(&SynchronizationObject);
@@ -612,22 +714,27 @@ FString FStreamingNetworkPlatformFile::ConvertToAbsolutePathForExternalAppForWri
 
 	// perform a local operation
 	FArrayReader Response;
+
 	if (SendPayloadAndReceiveResponse(Payload, Response) == true)
 	{
 		Response << RelativeFilename;
 	}
+
 	return RelativeFilename;
 }
 
+
 bool FStreamingNetworkPlatformFile::DirectoryExists(const TCHAR* Directory)
 {
 	// If there are any syncable files in this directory, consider it existing
 	FString RelativeDirectory = Directory;
 	FPaths::MakeStandardFilename(RelativeDirectory);
 	FServerTOC::FDirectory* ServerDirectory = ServerFiles.FindDirectory(RelativeDirectory);
-	return ServerDirectory != NULL;
+
+	return ServerDirectory != nullptr;
 }
 
+
 void FStreamingNetworkPlatformFile::GetFileInfo(const TCHAR* Filename, FFileInfo& Info)
 {
 	FScopeLock ScopeLock(&SynchronizationObject);
@@ -667,6 +774,7 @@ void FStreamingNetworkPlatformFile::GetFileInfo(const TCHAR* Filename, FFileInfo
 	}
 }
 
+
 FStreamingNetworkFileHandle* FStreamingNetworkPlatformFile::SendOpenMessage(const FString& Filename, bool bIsWriting, bool bAppend, bool bAllowRead)
 {
 	FScopeLock ScopeLock(&SynchronizationObject);
@@ -684,7 +792,7 @@ FStreamingNetworkFileHandle* FStreamingNetworkPlatformFile::SendOpenMessage(cons
 	FArrayReader Response;
 	if (SendPayloadAndReceiveResponse(Payload, Response) == false)
 	{
-		return NULL;
+		return nullptr;
 	}
 
 	// This server handle ID which will be used to perform operations on this file.
@@ -706,10 +814,11 @@ FStreamingNetworkFileHandle* FStreamingNetworkPlatformFile::SendOpenMessage(cons
 	}
 	else
 	{
-		return NULL;
+		return nullptr;
 	}
 }
 
+
 bool FStreamingNetworkPlatformFile::SendReadMessage(uint64 HandleId, uint8* Destination, int64 BytesToRead)
 {
 	FScopeLock ScopeLock(&SynchronizationObject);
@@ -721,6 +830,7 @@ bool FStreamingNetworkPlatformFile::SendReadMessage(uint64 HandleId, uint8* Dest
 
 	// Send the filename over
 	FArrayReader Response;
+
 	if (SendPayloadAndReceiveResponse(Payload, Response) == false)
 	{
 		return false;
@@ -731,6 +841,7 @@ bool FStreamingNetworkPlatformFile::SendReadMessage(uint64 HandleId, uint8* Dest
 	Response << ServerBytesRead;
 
 	bool bSuccess = (ServerBytesRead == BytesToRead);
+
 	if (bSuccess)
 	{
 		// Get the data.
@@ -740,6 +851,7 @@ bool FStreamingNetworkPlatformFile::SendReadMessage(uint64 HandleId, uint8* Dest
 	return bSuccess;
 }
 
+
 bool FStreamingNetworkPlatformFile::SendWriteMessage(uint64 HandleId, const uint8* Source, int64 BytesToWrite)
 {
 	FScopeLock ScopeLock(&SynchronizationObject);
@@ -752,6 +864,7 @@ bool FStreamingNetworkPlatformFile::SendWriteMessage(uint64 HandleId, const uint
 	Payload.Serialize(const_cast<uint8*>(Source), BytesToWrite);
 
 	FArrayReader Response;
+
 	if (SendPayloadAndReceiveResponse(Payload, Response) == false) 
 	{
 		return false;
@@ -764,6 +877,7 @@ bool FStreamingNetworkPlatformFile::SendWriteMessage(uint64 HandleId, const uint
 	return (ServerBytesWritten == BytesToWrite);
 }
 
+
 bool FStreamingNetworkPlatformFile::SendSeekMessage(uint64 HandleId, int64 NewPosition)
 {
 	FScopeLock ScopeLock(&SynchronizationObject);
@@ -774,6 +888,7 @@ bool FStreamingNetworkPlatformFile::SendSeekMessage(uint64 HandleId, int64 NewPo
 	Payload << NewPosition;
 
 	FArrayReader Response;
+
 	if (SendPayloadAndReceiveResponse(Payload, Response) == false)
 	{
 		return false;
@@ -781,9 +896,11 @@ bool FStreamingNetworkPlatformFile::SendSeekMessage(uint64 HandleId, int64 NewPo
 
 	int64 ServerNewPosition = -1;
 	Response << ServerNewPosition;
+
 	return (ServerNewPosition == NewPosition);
 }
 
+
 bool FStreamingNetworkPlatformFile::SendCloseMessage(uint64 HandleId)
 {
 	FScopeLock ScopeLock(&SynchronizationObject);
@@ -793,15 +910,18 @@ bool FStreamingNetworkPlatformFile::SendCloseMessage(uint64 HandleId)
 	Payload << HandleId;
 
 	FArrayReader Response;
+
 	return SendPayloadAndReceiveResponse(Payload, Response);
 }
 
+
 void FStreamingNetworkPlatformFile::PerformHeartbeat()
 {
 	FNetworkFileArchive Payload(NFS_Messages::Heartbeat);
 
 	// send the filename over
 	FArrayReader Response;
+
 	if(SendPayloadAndReceiveResponse(Payload,Response))
 	{
 		return;
@@ -810,16 +930,21 @@ void FStreamingNetworkPlatformFile::PerformHeartbeat()
     check(0);
 }
 
+
 /**
  * Module for the streaming file
  */
-class FStreamingFileModule : public IPlatformFileModule
+class FStreamingFileModule
+	: public IPlatformFileModule
 {
 public:
+
 	virtual IPlatformFile* GetPlatformFile() override
 	{
 		static TScopedPointer<IPlatformFile> AutoDestroySingleton(new FStreamingNetworkPlatformFile());
 		return AutoDestroySingleton.GetOwnedPointer();
 	}
 };
-IMPLEMENT_MODULE(FStreamingFileModule, StreamingFile);
\ No newline at end of file
+
+
+IMPLEMENT_MODULE(FStreamingFileModule, StreamingFile);
diff --git a/Engine/Source/Runtime/StreamingFile/Public/StreamingNetworkPlatformFile.h b/Engine/Source/Runtime/StreamingFile/Public/StreamingNetworkPlatformFile.h
--- a/Engine/Source/Runtime/StreamingFile/Public/StreamingNetworkPlatformFile.h
+++ b/Engine/Source/Runtime/StreamingFile/Public/StreamingNetworkPlatformFile.h
@@ -1,12 +1,14 @@
 // Copyright 1998-2015 Epic Games, Inc. All Rights Reserved.
+
 #pragma once
 
 #include "NetworkMessage.h"
-#include "ServerTOC.h"
 #include "NetworkPlatformFile.h"
 
+
 DECLARE_LOG_CATEGORY_EXTERN(LogStreamingPlatformFile, Log, All);
 
+
 #if 0
 /**
  * Visitor to gather local files with their timestamps
@@ -14,6 +16,7 @@ DECLARE_LOG_CATEGORY_EXTERN(LogStreamingPlatformFile, Log, All);
 class STREAMINGFILE_API FStreamingLocalTimestampVisitor : public IPlatformFile::FDirectoryVisitor
 {
 private:
+
 	/** The file interface to use for any file operations */
 	IPlatformFile& FileInterface;
 
@@ -27,6 +30,7 @@ private:
 	TArray<FString> DirectoriesToNotRecurse;
 
 public:
+
 	/** Relative paths to local files and their timestamps */
 	TMap<FString, FDateTime> FileTimes;
 	
@@ -110,16 +114,26 @@ public:
 };
 #endif
 
+
 /**
  * Wrapper to redirect the low level file system to a server
  */
-class STREAMINGFILE_API FStreamingNetworkPlatformFile : public FNetworkPlatformFile
+class STREAMINGFILE_API FStreamingNetworkPlatformFile
+	: public FNetworkPlatformFile
 {
 	friend class FAsyncFileSync;
 
 	// FNetworkPlatformFile interface
 	virtual bool InitializeInternal(IPlatformFile* Inner, const TCHAR* HostIP) override;
 
+public:
+
+	/** Default Constructor */
+	FStreamingNetworkPlatformFile() { };
+
+	/** Virtual destructor */
+	virtual ~FStreamingNetworkPlatformFile();
+
 public:
 
 	static const TCHAR* GetTypeName()
@@ -127,93 +141,67 @@ public:
 		return TEXT("StreamingFile");
 	}
 
-	/** Constructor */
-	FStreamingNetworkPlatformFile() {};
+	/** Sends Open message to the server and creates a new file handle if successful. */
+	class FStreamingNetworkFileHandle* SendOpenMessage(const FString& Filename, bool bIsWriting, bool bAppend, bool bAllowRead);
 
-	/** Destructor */
-	virtual ~FStreamingNetworkPlatformFile();
+	/** Sends Read message to the server. */
+	bool SendReadMessage(uint64 HandleId, uint8* Destination, int64 BytesToRead);
+
+	/** Sends Write message to the server. */
+	bool SendWriteMessage(uint64 HandleId, const uint8* Source, int64 BytesToWrite);	
+
+	/** Sends Seek message to the server. */
+	bool SendSeekMessage(uint64 HandleId, int64 NewPosition);	
+
+	/** Sends Close message to the server. */
+	bool SendCloseMessage(uint64 HandleId);
+
+public:
 
 	// need to override what FNetworkPlatformFile does here
 	void InitializeAfterSetActive() override { }
 
-	// Begin IPlatformFile Interface 
+	// IPlatformFile interface
+
 	virtual bool ShouldBeUsed(IPlatformFile* Inner, const TCHAR* CmdLine) const override;
 
 	virtual IPlatformFile* GetLowerLevel() override
 	{
-		return NULL;
+		return nullptr;
 	}
 
 	virtual const TCHAR* GetName() const override
 	{
 		return FStreamingNetworkPlatformFile::GetTypeName();
 	}
-	virtual bool		DeleteFile(const TCHAR* Filename) override;
-	virtual bool		IsReadOnly(const TCHAR* Filename) override
-	{
-		FFileInfo Info;
-		GetFileInfo(Filename, Info);
-		return Info.ReadOnly;
-	}
-	virtual bool		MoveFile(const TCHAR* To, const TCHAR* From) override;
-	virtual bool		SetReadOnly(const TCHAR* Filename, bool bNewReadOnlyValue) override;
-	virtual FDateTime	GetTimeStamp(const TCHAR* Filename) override
-	{
-		FFileInfo Info;
-		GetFileInfo(Filename, Info);
-		return Info.TimeStamp;
-	}
-	virtual void		SetTimeStamp(const TCHAR* Filename, FDateTime DateTime) override;
-	virtual FDateTime	GetAccessTimeStamp(const TCHAR* Filename) override
-	{
-		FFileInfo Info;
-		GetFileInfo(Filename, Info);
-		return Info.AccessTimeStamp;
-	}
-	virtual IFileHandle*	OpenRead(const TCHAR* Filename) override;
-	virtual IFileHandle*	OpenWrite(const TCHAR* Filename, bool bAppend, bool bAllowRead) override;
-	virtual bool		DirectoryExists(const TCHAR* Directory) override;
-	virtual bool		CreateDirectoryTree(const TCHAR* Directory) override;
-	virtual bool		CreateDirectory(const TCHAR* Directory) override;
-	virtual bool		DeleteDirectory(const TCHAR* Directory) override;
-
-	virtual bool		IterateDirectory(const TCHAR* Directory, IPlatformFile::FDirectoryVisitor& Visitor) override;
-	virtual bool		IterateDirectoryRecursively(const TCHAR* Directory, IPlatformFile::FDirectoryVisitor& Visitor) override;
-	virtual bool		DeleteDirectoryRecursively(const TCHAR* Directory) override;
-	virtual bool		CopyFile(const TCHAR* To, const TCHAR* From) override;
 
+	virtual bool DeleteFile(const TCHAR* Filename) override;
+	virtual bool IsReadOnly(const TCHAR* Filename) override;
+	virtual bool MoveFile(const TCHAR* To, const TCHAR* From) override;
+	virtual bool SetReadOnly(const TCHAR* Filename, bool bNewReadOnlyValue) override;
+	virtual FDateTime GetTimeStamp(const TCHAR* Filename) override;
+	virtual void SetTimeStamp(const TCHAR* Filename, FDateTime DateTime) override;
+	virtual FDateTime GetAccessTimeStamp(const TCHAR* Filename) override;
+	virtual IFileHandle* OpenRead(const TCHAR* Filename) override;
+	virtual IFileHandle* OpenWrite(const TCHAR* Filename, bool bAppend, bool bAllowRead) override;
+	virtual bool DirectoryExists(const TCHAR* Directory) override;
+	virtual bool CreateDirectoryTree(const TCHAR* Directory) override;
+	virtual bool CreateDirectory(const TCHAR* Directory) override;
+	virtual bool DeleteDirectory(const TCHAR* Directory) override;
+	virtual bool IterateDirectory(const TCHAR* Directory, IPlatformFile::FDirectoryVisitor& Visitor) override;
+	virtual bool IterateDirectoryRecursively(const TCHAR* Directory, IPlatformFile::FDirectoryVisitor& Visitor) override;
+	virtual bool DeleteDirectoryRecursively(const TCHAR* Directory) override;
+	virtual bool CopyFile(const TCHAR* To, const TCHAR* From) override;
 	virtual FString ConvertToAbsolutePathForExternalAppForRead( const TCHAR* Filename ) override;
 	virtual FString ConvertToAbsolutePathForExternalAppForWrite( const TCHAR* Filename ) override;
 
-	// End IPlatformFile Interface
-
-	/** Sends Open message to the server and creates a new file handle if successful. */
-	class FStreamingNetworkFileHandle* SendOpenMessage(const FString& Filename, bool bIsWriting, bool bAppend, bool bAllowRead);
-
-	/** Sends Read message to the server. */
-	bool SendReadMessage(uint64 HandleId, uint8* Destination, int64 BytesToRead);
-
-	/** Sends Write message to the server. */
-	bool SendWriteMessage(uint64 HandleId, const uint8* Source, int64 BytesToWrite);	
-
-	/** Sends Seek message to the server. */
-	bool SendSeekMessage(uint64 HandleId, int64 NewPosition);	
-
-	/** Sends Close message to the server. */
-	bool SendCloseMessage(uint64 HandleId);
-	
 private:
 
-	// Begin FNetworkPlatformFile interface
+	// FNetworkPlatformFile interface
 
 	virtual void PerformHeartbeat() override;
-
 	virtual void GetFileInfo(const TCHAR* Filename, FFileInfo& Info) override;
 
-	// End FNetworkPlatformFile interface
-
-
-
 private:
 
 	/** Set files that the server said we should sync asynchronously */
@@ -222,7 +210,3 @@ private:
 	/** Stored information about the files we have already cached */
 	TMap<FString, FFileInfo> CachedFileInfo;
 };
-
-
-
-
diff --git a/Engine/Source/Runtime/StreamingFile/StreamingFile.Build.cs b/Engine/Source/Runtime/StreamingFile/StreamingFile.Build.cs
--- a/Engine/Source/Runtime/StreamingFile/StreamingFile.Build.cs
+++ b/Engine/Source/Runtime/StreamingFile/StreamingFile.Build.cs
@@ -6,8 +6,12 @@ public class StreamingFile : ModuleRules
 {
 	public StreamingFile(TargetInfo Target)
 	{
-        PublicIncludePathModuleNames.Add("NetworkFile");
-
-        PrivateDependencyModuleNames.AddRange(new string[] { "Core", "Sockets", "NetworkFile" });
-	}
+        PublicDependencyModuleNames.AddRange(
+            new string[] {
+                "Core",
+                "NetworkFile",
+                "Sockets",
+            }
+        );
+    }
 }
\ No newline at end of file
