Fixed cook on the fly streaming server from not working because couldn't find plugins directory.

diff --git a/Engine/Source/Runtime/NetworkFile/Public/NetworkPlatformFile.h b/Engine/Source/Runtime/NetworkFile/Public/NetworkPlatformFile.h
--- a/Engine/Source/Runtime/NetworkFile/Public/NetworkPlatformFile.h
+++ b/Engine/Source/Runtime/NetworkFile/Public/NetworkPlatformFile.h
@@ -178,22 +178,13 @@ private:
 	 */
 	void EnsureFileIsLocal(const FString& Filename);
 
+protected:
 	/**
 	* Does normal path standardization, and also any extra modifications to make string comparisons against
 	* the internal directory list work properly.
 	*/
 	void MakeStandardNetworkFilename(FString& Filename);
 
-	/**
-	 * This function will send a payload data (with header) and wait for a response, serializing
-	 * the response to a FBufferArchive
-	 *
-	 * @param Payload Bytes to send over the network
-	 * @param Response The archive to read the response into
-	 *
-	 * @return true if successful
-	 */
-
 protected:
 
 	/** This is true after the DDC directories have been loaded from the DDC system */
diff --git a/Engine/Source/Runtime/StreamingFile/Private/StreamingNetworkPlatformFile.cpp b/Engine/Source/Runtime/StreamingFile/Private/StreamingNetworkPlatformFile.cpp
--- a/Engine/Source/Runtime/StreamingFile/Private/StreamingNetworkPlatformFile.cpp
+++ b/Engine/Source/Runtime/StreamingFile/Private/StreamingNetworkPlatformFile.cpp
@@ -374,7 +374,7 @@ bool FStreamingNetworkPlatformFile::DeleteFile(const TCHAR* Filename)
 {
 	FScopeLock ScopeLock(&SynchronizationObject);
 	FString RelativeFilename = Filename;
-	FPaths::MakeStandardFilename(RelativeFilename);
+	MakeStandardNetworkFilename(RelativeFilename);
 
 	FStreamingNetworkFileArchive Payload(NFS_Messages::DeleteFile);
 	Payload << RelativeFilename;
@@ -407,9 +407,9 @@ bool FStreamingNetworkPlatformFile::MoveFile(const TCHAR* To, const TCHAR* From)
 	FScopeLock ScopeLock(&SynchronizationObject);
 
 	FString RelativeFrom = From; 
-	FPaths::MakeStandardFilename(RelativeFrom);
+	MakeStandardNetworkFilename(RelativeFrom);
 	FString RelativeTo = To; 
-	FPaths::MakeStandardFilename(RelativeTo);
+	MakeStandardNetworkFilename(RelativeTo);
 
 	FStreamingNetworkFileArchive Payload(NFS_Messages::MoveFile);
 	Payload << RelativeFrom;
@@ -435,7 +435,7 @@ bool FStreamingNetworkPlatformFile::SetReadOnly(const TCHAR* Filename, bool bNew
 
 	FStreamingNetworkFileArchive Payload(NFS_Messages::SetReadOnly);
 	FString RelativeFilename = Filename;
-	FPaths::MakeStandardFilename(RelativeFilename);
+	MakeStandardNetworkFilename(RelativeFilename);
 
 	Payload << RelativeFilename;
 	Payload << bNewReadOnlyValue;
@@ -470,7 +470,7 @@ void FStreamingNetworkPlatformFile::SetTimeStamp(const TCHAR* Filename, FDateTim
 
 	FStreamingNetworkFileArchive Payload(NFS_Messages::SetTimeStamp);
 	FString RelativeFilename = Filename;
-	FPaths::MakeStandardFilename(RelativeFilename);
+	MakeStandardNetworkFilename(RelativeFilename);
 	Payload << RelativeFilename;
 	Payload << DateTime;
 
@@ -499,7 +499,7 @@ FDateTime FStreamingNetworkPlatformFile::GetAccessTimeStamp(const TCHAR* Filenam
 IFileHandle* FStreamingNetworkPlatformFile::OpenRead(const TCHAR* Filename, bool bAllowWrite)
 {
 	FString RelativeFilename = Filename;
-	FPaths::MakeStandardFilename(RelativeFilename);
+	MakeStandardNetworkFilename(RelativeFilename);
 
 	FStreamingNetworkFileHandle* FileHandle = SendOpenMessage(RelativeFilename, false, false, false);
 
@@ -510,7 +510,7 @@ IFileHandle* FStreamingNetworkPlatformFile::OpenRead(const TCHAR* Filename, bool
 IFileHandle* FStreamingNetworkPlatformFile::OpenWrite(const TCHAR* Filename, bool bAppend, bool bAllowRead)
 {
 	FString RelativeFilename = Filename;
-	FPaths::MakeStandardFilename(RelativeFilename);
+	MakeStandardNetworkFilename(RelativeFilename);
 	FStreamingNetworkFileHandle* FileHandle = SendOpenMessage(RelativeFilename, true, bAppend, bAllowRead);
 
 	return FileHandle;
@@ -528,7 +528,7 @@ bool FStreamingNetworkPlatformFile::CreateDirectory(const TCHAR* Directory)
 
 	FStreamingNetworkFileArchive Payload(NFS_Messages::CreateDirectory);
 	FString RelativeDirectory = Directory;
-	FPaths::MakeStandardFilename(RelativeDirectory);
+	MakeStandardNetworkFilename(RelativeDirectory);
 	Payload << RelativeDirectory;
 
 	// perform a local operation
@@ -552,7 +552,7 @@ bool FStreamingNetworkPlatformFile::DeleteDirectory(const TCHAR* Directory)
 
 	FStreamingNetworkFileArchive Payload(NFS_Messages::DeleteDirectory);
 	FString RelativeDirectory = Directory;
-	FPaths::MakeStandardFilename(RelativeDirectory);
+	MakeStandardNetworkFilename(RelativeDirectory);
 	Payload << RelativeDirectory;
 
 	// perform a local operation
@@ -573,7 +573,7 @@ bool FStreamingNetworkPlatformFile::DeleteDirectory(const TCHAR* Directory)
 bool FStreamingNetworkPlatformFile::IterateDirectory(const TCHAR* InDirectory, IPlatformFile::FDirectoryVisitor& Visitor)
 {
 	FString RelativeDirectory = InDirectory;
-	FPaths::MakeStandardFilename(RelativeDirectory);
+	MakeStandardNetworkFilename(RelativeDirectory);
 	// for .dll, etc searches that don't specify a path, we need to strip off the path
 	// before we send it to the visitor
 	bool bHadNoPath = InDirectory[0] == 0;
@@ -606,7 +606,7 @@ bool FStreamingNetworkPlatformFile::IterateDirectory(const TCHAR* InDirectory, I
 bool FStreamingNetworkPlatformFile::IterateDirectoryRecursively(const TCHAR* InDirectory, IPlatformFile::FDirectoryVisitor& Visitor)
 {
 	FString RelativeDirectory = InDirectory;
-	FPaths::MakeStandardFilename(RelativeDirectory);
+	MakeStandardNetworkFilename(RelativeDirectory);
 	// we loop until this is false
 	bool RetVal = true;
 
@@ -639,7 +639,7 @@ bool FStreamingNetworkPlatformFile::DeleteDirectoryRecursively(const TCHAR* Dire
 
 	FStreamingNetworkFileArchive Payload(NFS_Messages::DeleteDirectoryRecursively);
 	FString RelativeDirectory = Directory;
-	FPaths::MakeStandardFilename(RelativeDirectory);
+	MakeStandardNetworkFilename(RelativeDirectory);
 	Payload << RelativeDirectory;
 
 	// perform a local operation
@@ -662,8 +662,8 @@ bool FStreamingNetworkPlatformFile::CopyFile(const TCHAR* To, const TCHAR* From)
 	FScopeLock ScopeLock(&SynchronizationObject);
 
 	FStreamingNetworkFileArchive Payload(NFS_Messages::CopyFile);
-	FString RelativeTo = To; FPaths::MakeStandardFilename(RelativeTo);
-	FString RelativeFrom = From; FPaths::MakeStandardFilename(RelativeFrom);
+	FString RelativeTo = To; MakeStandardNetworkFilename(RelativeTo);
+	FString RelativeFrom = From; MakeStandardNetworkFilename(RelativeFrom);
 	Payload << RelativeTo;
 	Payload << RelativeFrom;
 
@@ -688,7 +688,7 @@ FString FStreamingNetworkPlatformFile::ConvertToAbsolutePathForExternalAppForRea
 
 	FStreamingNetworkFileArchive Payload(NFS_Messages::ToAbsolutePathForRead);
 	FString RelativeFilename = Filename;
-	FPaths::MakeStandardFilename(RelativeFilename);
+	MakeStandardNetworkFilename(RelativeFilename);
 	Payload << RelativeFilename;
 
 	// perform a local operation
@@ -709,7 +709,7 @@ FString FStreamingNetworkPlatformFile::ConvertToAbsolutePathForExternalAppForWri
 
 	FStreamingNetworkFileArchive Payload(NFS_Messages::ToAbsolutePathForWrite);
 	FString RelativeFilename = Filename; 
-	FPaths::MakeStandardFilename(RelativeFilename);
+	MakeStandardNetworkFilename(RelativeFilename);
 	Payload << RelativeFilename;
 
 	// perform a local operation
@@ -728,7 +728,7 @@ bool FStreamingNetworkPlatformFile::DirectoryExists(const TCHAR* Directory)
 {
 	// If there are any syncable files in this directory, consider it existing
 	FString RelativeDirectory = Directory;
-	FPaths::MakeStandardFilename(RelativeDirectory);
+	MakeStandardNetworkFilename(RelativeDirectory);
 	FServerTOC::FDirectory* ServerDirectory = ServerFiles.FindDirectory(RelativeDirectory);
 
 	return ServerDirectory != nullptr;
@@ -740,7 +740,7 @@ void FStreamingNetworkPlatformFile::GetFileInfo(const TCHAR* Filename, FFileInfo
 	FScopeLock ScopeLock(&SynchronizationObject);
 
 	FString RelativeFilename = Filename;
-	FPaths::MakeStandardFilename(RelativeFilename);
+	MakeStandardNetworkFilename(RelativeFilename);
 
 	if (!CachedFileInfo.Contains(RelativeFilename))
 	{
