Change the default cotf configuration on PS5 to streaming (FStreamingNetworkPlatformFile) instead of caching (FNetworkPlatformFile), due to slow file writes.
This means that specifying -streaming no longer has any effect on PS5, while -allowcaching can be used on PS5 to use the old behavior with a local cache.

#rnx

diff --git a/Engine/Source/Runtime/StreamingFile/Private/StreamingNetworkPlatformFile.cpp b/Engine/Source/Runtime/StreamingFile/Private/StreamingNetworkPlatformFile.cpp
--- a/Engine/Source/Runtime/StreamingFile/Private/StreamingNetworkPlatformFile.cpp
+++ b/Engine/Source/Runtime/StreamingFile/Private/StreamingNetworkPlatformFile.cpp
@@ -297,17 +297,17 @@ bool FStreamingNetworkPlatformFile::ShouldBeUsed(IPlatformFile* Inner, const TCH
 		bResult = FParse::Param(CmdLine, TEXT("Streaming")) || !FPlatformMisc::SupportsLocalCaching();
 		if (FPlatformMisc::AllowLocalCaching())
 		{
-			//@todo. Platforms: Do we need to support 'in place' non-streaming cookonthefly?
 			// On desktop platforms, assume 'in place' execution - to prevent deletion of content, force streaming,
 			// unless it's explicitly allowed with -allowcaching (for automation tests with staged builds).
 			const bool bAllowCaching = FParse::Param(CmdLine, TEXT("AllowCaching"));
 
 			if (bResult == false && !bAllowCaching)
 			{
-				UE_LOG(LogStreamingPlatformFile, Warning, TEXT("Cooked desktop platforms do not support non-streaming. Forcing streaming on."));
+				UE_LOG(LogStreamingPlatformFile, Display, TEXT("Platform supports local caching, but requires explicitly specifying -AllowCaching to use it."));
 			}
 			bResult = bResult || !bAllowCaching;
 		}
+		UE_CLOG(bResult, LogStreamingPlatformFile, Display, TEXT("Using streaming network file system."));
 	}
 
 	return bResult;
