Removing file info caching from StreamingNetworkPlatformFile.

- File information caching was return the wrong size for files that have changed since we first cached the information.  This specifically happens when iterating on shaders (the global shader map) using ODSC (on demand shader compilation) and COTF (cook on the fly).
- The easiest solution is to removing the caching all together so we always get the most up-to-date size.
- In the future if we need the caching back (optimization) we can investigating the cooking server telling the console to update its cache when files change.

diff --git a/Engine/Source/Runtime/StreamingFile/Private/StreamingNetworkPlatformFile.cpp b/Engine/Source/Runtime/StreamingFile/Private/StreamingNetworkPlatformFile.cpp
--- a/Engine/Source/Runtime/StreamingFile/Private/StreamingNetworkPlatformFile.cpp
+++ b/Engine/Source/Runtime/StreamingFile/Private/StreamingNetworkPlatformFile.cpp
@@ -751,36 +751,22 @@ void FStreamingNetworkPlatformFile::GetFileInfo(const TCHAR* Filename, FFileInfo
 	FString RelativeFilename = Filename;
 	MakeStandardNetworkFilename(RelativeFilename);
 
-	if (!CachedFileInfo.Contains(RelativeFilename))
-	{
-		FStreamingNetworkFileArchive Payload(NFS_Messages::GetFileInfo);
-		Payload << const_cast<FString&>(RelativeFilename);
-
-//		if (RelativeFilename == TEXT("../../../UDKGame/Content/Maps/GDC12_Ice/GDC_2012_Throne_Cave.uasset"))
-//		{
-//			Info.ReadOnly = true;
-//		}
+	FStreamingNetworkFileArchive Payload(NFS_Messages::GetFileInfo);
+	Payload << const_cast<FString&>(RelativeFilename);
 
-		// Send the filename over
-		FArrayReader Response;
-		if (SendPayloadAndReceiveResponse(Payload, Response) == false)
-		{
-			return;
-		}
-
-		// Get info from the response
-		Response << Info.FileExists;
-		Response << Info.ReadOnly;
-		Response << Info.Size;
-		Response << Info.TimeStamp;
-		Response << Info.AccessTimeStamp;
-
-		CachedFileInfo.Add(RelativeFilename, Info);
-	}
-	else
+	// Send the filename over
+	FArrayReader Response;
+	if (SendPayloadAndReceiveResponse(Payload, Response) == false)
 	{
-		Info = *(CachedFileInfo.Find(RelativeFilename));
+		return;
 	}
+
+	// Get info from the response
+	Response << Info.FileExists;
+	Response << Info.ReadOnly;
+	Response << Info.Size;
+	Response << Info.TimeStamp;
+	Response << Info.AccessTimeStamp;
 }
 
 
diff --git a/Engine/Source/Runtime/StreamingFile/Public/StreamingNetworkPlatformFile.h b/Engine/Source/Runtime/StreamingFile/Public/StreamingNetworkPlatformFile.h
--- a/Engine/Source/Runtime/StreamingFile/Public/StreamingNetworkPlatformFile.h
+++ b/Engine/Source/Runtime/StreamingFile/Public/StreamingNetworkPlatformFile.h
@@ -216,7 +216,4 @@ private:
 
 	/** Set files that the server said we should sync asynchronously */
 	TArray<FString> FilesToSyncAsync;
-
-	/** Stored information about the files we have already cached */
-	TMap<FString, FFileInfo> CachedFileInfo;
 };
